UX improvements for the in-editor culture picker

- Added auto-expansion for filtered entries with valid children.
- Only show entries for languages we directly support.

diff --git a/Engine/Source/Editor/InternationalizationSettings/Private/InternationalizationSettingsModelDetails.cpp b/Engine/Source/Editor/InternationalizationSettings/Private/InternationalizationSettingsModelDetails.cpp
--- a/Engine/Source/Editor/InternationalizationSettings/Private/InternationalizationSettingsModelDetails.cpp
+++ b/Engine/Source/Editor/InternationalizationSettings/Private/InternationalizationSettingsModelDetails.cpp
@@ -58,13 +58,15 @@ namespace
 
 		FLocalizedCulturesFlyweight()
 		{
+			constexpr bool bIncludeDerivedCultures = false;
+
 			{
 				const TArray<FString> LocalizedCultureNames = FTextLocalizationManager::Get().GetLocalizedCultureNames(ELocalizationLoadFlags::Editor);
-				LocalizedCulturesForEditor = FInternationalization::Get().GetAvailableCultures(LocalizedCultureNames, true);
+				LocalizedCulturesForEditor = FInternationalization::Get().GetAvailableCultures(LocalizedCultureNames, bIncludeDerivedCultures);
 			}
 			{
 				const TArray<FString> LocalizedCultureNames = FTextLocalizationManager::Get().GetLocalizedCultureNames(ELocalizationLoadFlags::Game);
-				LocalizedCulturesForGame = FInternationalization::Get().GetAvailableCultures(LocalizedCultureNames, true);
+				LocalizedCulturesForGame = FInternationalization::Get().GetAvailableCultures(LocalizedCultureNames, bIncludeDerivedCultures);
 			}
 		}
 	};
@@ -170,7 +172,7 @@ namespace
 				.DisplayNameFormat(SCulturePicker::ECultureDisplayFormat::ActiveAndNativeCultureDisplayName);
 
 			return SNew(SBox)
-				.MaxDesiredHeight(300.0f)
+				.MaxDesiredHeight(500.0f)
 				.WidthOverride(300.0f)
 				[
 					CulturePicker
@@ -358,7 +360,7 @@ namespace
 				return false;
 			};
 			return SNew(SBox)
-				.MaxDesiredHeight(300.0f)
+				.MaxDesiredHeight(500.0f)
 				.WidthOverride(300.0f)
 				[
 					SNew(SCulturePicker)
diff --git a/Engine/Source/Editor/InternationalizationSettings/Private/SCulturePicker.cpp b/Engine/Source/Editor/InternationalizationSettings/Private/SCulturePicker.cpp
--- a/Engine/Source/Editor/InternationalizationSettings/Private/SCulturePicker.cpp
+++ b/Engine/Source/Editor/InternationalizationSettings/Private/SCulturePicker.cpp
@@ -59,17 +59,13 @@ void SCulturePicker::Construct( const FArguments& InArgs )
 		];
 
 
-	const auto& IsInitialSelection = [&](const TSharedPtr<FCultureEntry>& Entry) -> bool
-	{
-		return Entry->Culture == InArgs._InitialSelection;
-	};
-
-	const TSharedPtr<FCultureEntry>* InitialSelection = RootEntries.FindByPredicate(IsInitialSelection);
-	if (InitialSelection)
+	if (TSharedPtr<FCultureEntry> InitialSelection = FindEntryForCulture(InArgs._InitialSelection))
 	{
 		TGuardValue<bool> SuppressSelectionGuard(SuppressSelectionCallback, true);
-		TreeView->SetSelection(*InitialSelection);
+		TreeView->SetSelection(InitialSelection);
 	}
+
+	AutoExpandEntries();
 }
 
 void SCulturePicker::RequestTreeRefresh()
@@ -79,6 +75,53 @@ void SCulturePicker::RequestTreeRefresh()
 	{
 		TreeView->RequestTreeRefresh();
 	}
+	AutoExpandEntries();
+}
+
+TSharedPtr<FCultureEntry> SCulturePicker::FindEntryForCulture(FCulturePtr Culture) const
+{
+	return FindEntryForCultureImpl(Culture, RootEntries);
+}
+
+TSharedPtr<FCultureEntry> SCulturePicker::FindEntryForCultureImpl(FCulturePtr Culture, const TArray<TSharedPtr<FCultureEntry>>& Entries) const
+{
+	for (const TSharedPtr<FCultureEntry>& Entry : Entries)
+	{
+		if (Entry->Culture == Culture)
+		{
+			return Entry;
+		}
+		if (TSharedPtr<FCultureEntry> ChildEntry = FindEntryForCultureImpl(Culture, Entry->Children))
+		{
+			return ChildEntry;
+		}
+	}
+	return nullptr;
+}
+
+void SCulturePicker::AutoExpandEntries()
+{
+	return AutoExpandEntriesImpl(RootEntries);
+}
+
+void SCulturePicker::AutoExpandEntriesImpl(const TArray<TSharedPtr<FCultureEntry>>& Entries)
+{
+	if (TreeView)
+	{
+		for (const TSharedPtr<FCultureEntry>& Entry : Entries)
+		{
+			if (Entry->AutoExpand)
+			{
+				TreeView->SetItemExpansion(Entry, true);
+				
+				// Only recurse when actually filtering
+				if (!FilterString.IsEmpty())
+				{
+					AutoExpandEntriesImpl(Entry->Children);
+				}
+			}
+		}
+	}
 }
 
 void SCulturePicker::BuildStockEntries()
@@ -199,6 +242,7 @@ void SCulturePicker::RebuildEntries()
 			// If has children, must be added. If it is not filtered and it is pickable, should be added.
 			if ( OutEntry->Children.Num() != 0 || (!IsFilteredOut && IsPickable) )
 			{
+				OutEntry->AutoExpand = (!IsPickable || IsFilteredOut) && OutEntry->Children.Num() > 0;
 				OutEntries.Add(OutEntry);
 			}
 		}	
diff --git a/Engine/Source/Editor/InternationalizationSettings/Public/SCulturePicker.h b/Engine/Source/Editor/InternationalizationSettings/Public/SCulturePicker.h
--- a/Engine/Source/Editor/InternationalizationSettings/Public/SCulturePicker.h
+++ b/Engine/Source/Editor/InternationalizationSettings/Public/SCulturePicker.h
@@ -32,6 +32,7 @@ struct FCultureEntry
 	FCulturePtr Culture;
 	TArray< TSharedPtr<FCultureEntry> > Children;
 	bool IsSelectable;
+	bool AutoExpand = false;
 
 	FCultureEntry(const FCulturePtr& InCulture, const bool InIsSelectable = true)
 		: Culture(InCulture)
@@ -92,6 +93,12 @@ public:
 	void RequestTreeRefresh();
 
 private:
+	TSharedPtr<FCultureEntry> FindEntryForCulture(FCulturePtr Culture) const;
+	TSharedPtr<FCultureEntry> FindEntryForCultureImpl(FCulturePtr Culture, const TArray<TSharedPtr<FCultureEntry>>& Entries) const;
+
+	void AutoExpandEntries();
+	void AutoExpandEntriesImpl(const TArray<TSharedPtr<FCultureEntry>>& Entries);
+
 	void BuildStockEntries();
 	void RebuildEntries();
 
