Added a "flat list" option to the in-editor culture picker

#rnx

diff --git a/Engine/Source/Editor/InternationalizationSettings/Private/InternationalizationSettingsModelDetails.cpp b/Engine/Source/Editor/InternationalizationSettings/Private/InternationalizationSettingsModelDetails.cpp
--- a/Engine/Source/Editor/InternationalizationSettings/Private/InternationalizationSettingsModelDetails.cpp
+++ b/Engine/Source/Editor/InternationalizationSettings/Private/InternationalizationSettingsModelDetails.cpp
@@ -169,11 +169,12 @@ namespace
 				.InitialSelection(EditorLanguage)
 				.OnSelectionChanged_Lambda(OnSelectionChangedLambda)
 				.IsCulturePickable_Lambda(IsCulturePickableLambda)
-				.DisplayNameFormat(SCulturePicker::ECultureDisplayFormat::ActiveAndNativeCultureDisplayName);
+				.DisplayNameFormat(SCulturePicker::ECultureDisplayFormat::ActiveAndNativeCultureDisplayName)
+				.ViewMode(SCulturePicker::ECulturesViewMode::Flat);
 
 			return SNew(SBox)
 				.MaxDesiredHeight(500.0f)
-				.WidthOverride(300.0f)
+				.WidthOverride(350.0f)
 				[
 					CulturePicker
 				];
@@ -259,8 +260,8 @@ namespace
 				.DisplayNameFormat(SCulturePicker::ECultureDisplayFormat::ActiveAndNativeCultureDisplayName);
 
 			return SNew(SBox)
-				.MaxDesiredHeight(300.0f)
-				.WidthOverride(300.0f)
+				.MaxDesiredHeight(500.0f)
+				.WidthOverride(350.0f)
 				[
 					CulturePicker
 				];
@@ -361,12 +362,14 @@ namespace
 			};
 			return SNew(SBox)
 				.MaxDesiredHeight(500.0f)
-				.WidthOverride(300.0f)
+				.WidthOverride(350.0f)
 				[
 					SNew(SCulturePicker)
 					.InitialSelection(PreviewGameCulture)
 					.OnSelectionChanged_Lambda(CulturePickerSelectLambda)
 					.IsCulturePickable_Lambda(CulturePickerIsPickableLambda)
+					.DisplayNameFormat(SCulturePicker::ECultureDisplayFormat::ActiveAndNativeCultureDisplayName)
+					.ViewMode(SCulturePicker::ECulturesViewMode::Flat)
 					.CanSelectNone(true)
 				];
 		}
diff --git a/Engine/Source/Editor/InternationalizationSettings/Private/SCulturePicker.cpp b/Engine/Source/Editor/InternationalizationSettings/Private/SCulturePicker.cpp
--- a/Engine/Source/Editor/InternationalizationSettings/Private/SCulturePicker.cpp
+++ b/Engine/Source/Editor/InternationalizationSettings/Private/SCulturePicker.cpp
@@ -30,6 +30,7 @@ void SCulturePicker::Construct( const FArguments& InArgs )
 	OnCultureSelectionChanged = InArgs._OnSelectionChanged;
 	IsCulturePickable = InArgs._IsCulturePickable;
 	DisplayNameFormat = InArgs._DisplayNameFormat;
+	ViewMode = InArgs._ViewMode;
 	CanSelectNone = InArgs._CanSelectNone;
 
 	BuildStockEntries();
@@ -106,7 +107,7 @@ void SCulturePicker::AutoExpandEntries()
 
 void SCulturePicker::AutoExpandEntriesImpl(const TArray<TSharedPtr<FCultureEntry>>& Entries)
 {
-	if (TreeView)
+	if (ViewMode == ECulturesViewMode::Hierarchical && TreeView)
 	{
 		for (const TSharedPtr<FCultureEntry>& Entry : Entries)
 		{
@@ -239,11 +240,30 @@ void SCulturePicker::RebuildEntries()
 				IsFilteredOut = !Name.Contains(FilterString) && !DisplayName.Contains(FilterString) && !NativeName.Contains(FilterString);
 			}
 
-			// If has children, must be added. If it is not filtered and it is pickable, should be added.
-			if ( OutEntry->Children.Num() != 0 || (!IsFilteredOut && IsPickable) )
+			switch (ViewMode)
 			{
-				OutEntry->AutoExpand = (!IsPickable || IsFilteredOut) && OutEntry->Children.Num() > 0;
-				OutEntries.Add(OutEntry);
+			case ECulturesViewMode::Hierarchical:
+				// If has children, must be added. If it is not filtered and it is pickable, should be added.
+				if (OutEntry->Children.Num() != 0 || (!IsFilteredOut && IsPickable))
+				{
+					OutEntry->AutoExpand = (!IsPickable || IsFilteredOut) && OutEntry->Children.Num() > 0;
+					OutEntries.Add(OutEntry);
+				}
+				break;
+
+			case ECulturesViewMode::Flat:
+				// Add this entry only if it's valid, but always append its children (as they have already been validated)
+				if (IsPickable && !IsFilteredOut)
+				{
+					OutEntries.Add(OutEntry);
+				}
+				OutEntries.Append(OutEntry->Children);
+				OutEntry->Children.Reset();
+				break;
+
+			default:
+				checkf(false, TEXT("Unknown ECulturesViewMode!"));
+				break;
 			}
 		}	
 	};
@@ -322,7 +342,7 @@ FString SCulturePicker::GetCultureDisplayName(const FCultureRef& Culture, const
 	{
 		// Only show both names if they're different (to avoid repetition), and we're a root item (to avoid noise)
 		return (bIsRootItem && !NativeName.Equals(DisplayName, ESearchCase::CaseSensitive))
-			? FString::Printf(TEXT("%s (%s)"), *DisplayName, *NativeName)
+			? FString::Printf(TEXT("%s [%s]"), *DisplayName, *NativeName)
 			: DisplayName;
 	}
 
@@ -330,7 +350,7 @@ FString SCulturePicker::GetCultureDisplayName(const FCultureRef& Culture, const
 	{
 		// Only show both names if they're different (to avoid repetition), and we're a root item (to avoid noise)
 		return (bIsRootItem && !NativeName.Equals(DisplayName, ESearchCase::CaseSensitive))
-			? FString::Printf(TEXT("%s (%s)"), *NativeName, *DisplayName)
+			? FString::Printf(TEXT("%s [%s]"), *NativeName, *DisplayName)
 			: NativeName;
 	}
 
diff --git a/Engine/Source/Editor/InternationalizationSettings/Public/SCulturePicker.h b/Engine/Source/Editor/InternationalizationSettings/Public/SCulturePicker.h
--- a/Engine/Source/Editor/InternationalizationSettings/Public/SCulturePicker.h
+++ b/Engine/Source/Editor/InternationalizationSettings/Public/SCulturePicker.h
@@ -71,15 +71,25 @@ public:
 		NativeAndActiveCultureDisplayName,
 	};
 
+	enum class ECulturesViewMode
+	{
+		/** Display the cultures hierarchically in a tree */
+		Hierarchical,
+		/** Display the cultures as a flat list */
+		Flat,
+	};
+
 public:
 	SLATE_BEGIN_ARGS( SCulturePicker )
 		: _DisplayNameFormat(ECultureDisplayFormat::ActiveCultureDisplayName)
+		, _ViewMode(ECulturesViewMode::Hierarchical)
 		, _CanSelectNone(false)
 	{}
 		SLATE_EVENT( FOnSelectionChanged, OnSelectionChanged )
 		SLATE_EVENT( FIsCulturePickable, IsCulturePickable )
 		SLATE_ARGUMENT( FCulturePtr, InitialSelection )
 		SLATE_ARGUMENT(ECultureDisplayFormat, DisplayNameFormat)
+		SLATE_ARGUMENT(ECulturesViewMode, ViewMode)
 		SLATE_ARGUMENT(bool, CanSelectNone)
 	SLATE_END_ARGS()
 
@@ -130,6 +140,9 @@ private:
 	/** How should we display culture names? */
 	ECultureDisplayFormat DisplayNameFormat;
 
+	/** How should we display the list of cultures? */
+	ECulturesViewMode ViewMode;
+
 	/** Should a null culture option be available? */
 	bool CanSelectNone;
 
