Added missing NotifyFinishedChangingProperties calls

#rnx

diff --git a/Engine/Source/Editor/InternationalizationSettings/Private/UserGeneratedContentLocalizationDescriptorDetails.cpp b/Engine/Source/Editor/InternationalizationSettings/Private/UserGeneratedContentLocalizationDescriptorDetails.cpp
--- a/Engine/Source/Editor/InternationalizationSettings/Private/UserGeneratedContentLocalizationDescriptorDetails.cpp
+++ b/Engine/Source/Editor/InternationalizationSettings/Private/UserGeneratedContentLocalizationDescriptorDetails.cpp
@@ -183,6 +183,7 @@ void FUserGeneratedContentLocalizationDescriptorDetails::CulturesToGenerate_OnPr
 void FUserGeneratedContentLocalizationDescriptorDetails::CulturesToGenerate_OnPostBatchSelect()
 {
 	CulturesToGenerateHandle->NotifyPostChange(EPropertyChangeType::ValueSet);
+	CulturesToGenerateHandle->NotifyFinishedChangingProperties();
 	CulturesToGenerate_IsInBatchSelectOperation = false;
 }
 
@@ -210,6 +211,7 @@ void FUserGeneratedContentLocalizationDescriptorDetails::CulturesToGenerate_OnCu
 	if (!CulturesToGenerate_IsInBatchSelectOperation)
 	{
 		CulturesToGenerateHandle->NotifyPostChange(IsSelected ? EPropertyChangeType::ArrayAdd : EPropertyChangeType::ArrayRemove);
+		CulturesToGenerateHandle->NotifyFinishedChangingProperties();
 	}
 }
 
