Sort SCulturePicker display names in a locale correct way

#rnx

diff --git a/Engine/Source/Editor/InternationalizationSettings/Private/SCulturePicker.cpp b/Engine/Source/Editor/InternationalizationSettings/Private/SCulturePicker.cpp
--- a/Engine/Source/Editor/InternationalizationSettings/Private/SCulturePicker.cpp
+++ b/Engine/Source/Editor/InternationalizationSettings/Private/SCulturePicker.cpp
@@ -201,7 +201,7 @@ void SCulturePicker::BuildStockEntries()
 	{
 		const FString LHSDisplayName = GetCultureDisplayName(LHS->Culture.ToSharedRef(), false);
 		const FString RHSDisplayName = GetCultureDisplayName(RHS->Culture.ToSharedRef(), false);
-		return LHSDisplayName < RHSDisplayName;
+		return FTextComparison::CompareTo(LHSDisplayName, RHSDisplayName) < 0;
 	});
 }
 
@@ -306,7 +306,7 @@ void SCulturePicker::OnGetChildren(TSharedPtr<FCultureEntry> Entry, TArray< TSha
 	{
 		const FString LHSDisplayName = GetCultureDisplayName(LHS->Culture.ToSharedRef(), false);
 		const FString RHSDisplayName = GetCultureDisplayName(RHS->Culture.ToSharedRef(), false);
-		return LHSDisplayName < RHSDisplayName;
+		return FTextComparison::CompareTo(LHSDisplayName, RHSDisplayName) < 0;
 	});
 }
 
