Sort UGC Localization culture display names in a locale correct way

#rnx

diff --git a/Engine/Source/Editor/InternationalizationSettings/Private/UserGeneratedContentLocalizationDescriptorDetails.cpp b/Engine/Source/Editor/InternationalizationSettings/Private/UserGeneratedContentLocalizationDescriptorDetails.cpp
--- a/Engine/Source/Editor/InternationalizationSettings/Private/UserGeneratedContentLocalizationDescriptorDetails.cpp
+++ b/Engine/Source/Editor/InternationalizationSettings/Private/UserGeneratedContentLocalizationDescriptorDetails.cpp
@@ -50,7 +50,7 @@ void FUserGeneratedContentLocalizationDescriptorDetails::CustomizeChildren(TShar
 		// Sort by the current display name so that the order in the "native language" picker and the "cultures to generate" list matches
 		AvailableCultures.StableSort([](FCulturePtr One, FCulturePtr Two)
 		{
-			return One->GetDisplayName() < Two->GetDisplayName();
+			return FTextComparison::CompareTo(One->GetDisplayName(), Two->GetDisplayName()) < 0;
 		});
 	}
 
