Fixed 'Native Language' combo not closing when changing the language

#rnx

diff --git a/Engine/Source/Editor/InternationalizationSettings/Private/UserGeneratedContentLocalizationDescriptorDetails.cpp b/Engine/Source/Editor/InternationalizationSettings/Private/UserGeneratedContentLocalizationDescriptorDetails.cpp
--- a/Engine/Source/Editor/InternationalizationSettings/Private/UserGeneratedContentLocalizationDescriptorDetails.cpp
+++ b/Engine/Source/Editor/InternationalizationSettings/Private/UserGeneratedContentLocalizationDescriptorDetails.cpp
@@ -90,7 +90,7 @@ void FUserGeneratedContentLocalizationDescriptorDetails::CustomizeNativeCulture(
 		]
 		.ValueContent()
 		[
-			SNew(SComboButton)
+			SAssignNew(NativeCulture_ComboButton,SComboButton)
 			.ButtonContent()
 			[
 				SNew(STextBlock)
@@ -141,6 +141,11 @@ bool FUserGeneratedContentLocalizationDescriptorDetails::NativeCulture_IsCulture
 void FUserGeneratedContentLocalizationDescriptorDetails::NativeCulture_OnSelectionChanged(FCulturePtr SelectedCulture, ESelectInfo::Type SelectInfo)
 {
 	NativeCultureHandle->SetValue(SelectedCulture->GetName());
+
+	if (NativeCulture_ComboButton)
+	{
+		NativeCulture_ComboButton->SetIsOpen(false);
+	}
 }
 
 void FUserGeneratedContentLocalizationDescriptorDetails::CustomizeCulturesToGenerate(TSharedRef<IPropertyHandle> PropertyHandle, IDetailChildrenBuilder& StructBuilder, IPropertyTypeCustomizationUtils& StructCustomizationUtils)
diff --git a/Engine/Source/Editor/InternationalizationSettings/Private/UserGeneratedContentLocalizationDescriptorDetails.h b/Engine/Source/Editor/InternationalizationSettings/Private/UserGeneratedContentLocalizationDescriptorDetails.h
--- a/Engine/Source/Editor/InternationalizationSettings/Private/UserGeneratedContentLocalizationDescriptorDetails.h
+++ b/Engine/Source/Editor/InternationalizationSettings/Private/UserGeneratedContentLocalizationDescriptorDetails.h
@@ -7,6 +7,8 @@
 
 namespace ESelectInfo { enum Type : int; }
 
+class SComboButton;
+
 class FUserGeneratedContentLocalizationDescriptorDetails : public IPropertyTypeCustomization
 {
 public:
@@ -42,5 +44,7 @@ private:
 
 	TSharedPtr<IPropertyHandle> CulturesToGenerateHandle;
 
+	TSharedPtr<SComboButton> NativeCulture_ComboButton;
+
 	bool CulturesToGenerate_IsInBatchSelectOperation = false;
 };
