UE-230129 GitHub 12524 : Timezones have changed in Russia since 2014 and 2016

#tests Check the new timezones in the corresponding combo box

diff --git a/Engine/Source/Editor/InternationalizationSettings/Classes/InternationalizationSettingsModel.h b/Engine/Source/Editor/InternationalizationSettings/Classes/InternationalizationSettingsModel.h
--- a/Engine/Source/Editor/InternationalizationSettings/Classes/InternationalizationSettingsModel.h
+++ b/Engine/Source/Editor/InternationalizationSettings/Classes/InternationalizationSettingsModel.h
@@ -69,9 +69,9 @@ enum class ETimezoneSetting : uint8
 	Kuwait_Riyadh UMETA(DisplayName="(UTC+03:00) Kuwait, Riyadh"),
 	Baghdad UMETA(DisplayName="(UTC+03:00) Baghdad"),
 	Nairobi UMETA(DisplayName="(UTC+03:00) Nairobi"),
-	Kaliningrad UMETA(DisplayName="(UTC+03:00) Kaliningrad"),
+	Kaliningrad UMETA(DisplayName="(UTC+02:00) Kaliningrad"),
 	Tehran UMETA(DisplayName="(UTC+03:30) Tehran"),
-	Moscow_StPetersburg_Volgograd UMETA(DisplayName="(UTC+04:00) Moscow, St. Petersburg, Volgograd"),
+	Moscow_StPetersburg_Volgograd UMETA(DisplayName="(UTC+03:00) Moscow, St. Petersburg, Volgograd"),
 	AbuDhabi_Muscat UMETA(DisplayName="(UTC+04:00) Abu Dhabi, Muscat"),
 	Baku UMETA(DisplayName="(UTC+04:00) Baku"),
 	Yerevan UMETA(DisplayName="(UTC+04:00) Yerevan"),
@@ -83,31 +83,31 @@ enum class ETimezoneSetting : uint8
 	Chennai_Kolkata_Mumbai_NewDelhi UMETA(DisplayName="(UTC+05:30) Chennai, Kolkata, Mumbai, New Delhi"),
 	SriJayawardenepura UMETA(DisplayName="(UTC+05:30) Sri Jayawardenepura"),
 	Kathmandu UMETA(DisplayName="(UTC+05:45) Kathmandu"),
-	Ekaterinburg UMETA(DisplayName="(UTC+06:00) Ekaterinburg"),
+	Ekaterinburg UMETA(DisplayName="(UTC+05:00) Ekaterinburg"),
 	Astana UMETA(DisplayName="(UTC+06:00) Astana"),
 	Dhaka UMETA(DisplayName="(UTC+06:00) Dhaka"),
 	Yangon_Rangoon UMETA(DisplayName="(UTC+06:30) Yangon (Rangoon)"),
-	Novosibirsk UMETA(DisplayName="(UTC+07:00) Novosibirsk"),
 	Bangkok_Hanoi_Jakarta UMETA(DisplayName="(UTC+07:00) Bangkok, Hanoi, Jakarta"),
-	Krasnoyarsk UMETA(DisplayName="(UTC+08:00) Krasnoyarsk"),
+	Krasnoyarsk UMETA(DisplayName="(UTC+07:00) Krasnoyarsk"),
+	Novosibirsk UMETA(DisplayName="(UTC+07:00) Novosibirsk"),
 	Beijing_Chongqing_HongKong_Urumqi UMETA(DisplayName="(UTC+08:00) Beijing, Chongqing, Hong Kong, Urumqi"),
 	KualaLumpur_Singapore UMETA(DisplayName="(UTC+08:00) Kuala Lumpur, Singapore"),
 	Taipei UMETA(DisplayName="(UTC+08:00) Taipei"),
 	Perth UMETA(DisplayName="(UTC+08:00) Perth"),
 	Ulaanbaatar UMETA(DisplayName="(UTC+08:00) Ulaanbaatar"),
-	Irkutsk UMETA(DisplayName="(UTC+09:00) Irkutsk"),
+	Irkutsk UMETA(DisplayName="(UTC+08:00) Irkutsk"),
 	Seoul UMETA(DisplayName="(UTC+09:00) Seoul"),
 	Osaka_Sapporo_Tokyo UMETA(DisplayName="(UTC+09:00) Osaka, Sapporo, Tokyo"),
 	Darwin UMETA(DisplayName="(UTC+09:30) Darwin"),
 	Adelaide UMETA(DisplayName="(UTC+09:30) Adelaide"),
-	Yakutsk UMETA(DisplayName="(UTC+10:00) Yakutsk"),
+	Yakutsk UMETA(DisplayName="(UTC+9:00) Yakutsk"),
 	Canberra_Melbourne_Sydney UMETA(DisplayName="(UTC+10:00) Canberra, Melbourne, Sydney"),
 	Brisbane UMETA(DisplayName="(UTC+10:00) Brisbane"),
 	Hobart UMETA(DisplayName="(UTC+10:00) Hobart"),
 	Guam_PortMoresby UMETA(DisplayName="(UTC+10:00) Guam, Port Moresby"),
-	Vladivostok UMETA(DisplayName="(UTC+11:00) Vladivostok"),
+	Vladivostok UMETA(DisplayName="(UTC+10:00) Vladivostok"),
 	SolomonIs_NewCaledonia UMETA(DisplayName="(UTC+11:00) Solomon Is., New Caledonia"),
-	Magadan UMETA(DisplayName="(UTC+12:00) Magadan"),
+	Magadan UMETA(DisplayName="(UTC+11:00) Magadan"),
 	Fiji UMETA(DisplayName="(UTC+12:00) Fiji"),
 	Auckland_Wellington UMETA(DisplayName="(UTC+12:00) Auckland, Wellington"),
 	CoordinatedUniversalTime12 UMETA(DisplayName="(UTC+12:00) Coordinated Universal Time +12"),
