Added SCulturePickerCombo and updated existing copy/paste usages of this pattern to use the new widget

#rnx

diff --git a/Engine/Source/Editor/InternationalizationSettings/Private/InternationalizationSettingsModelDetails.cpp b/Engine/Source/Editor/InternationalizationSettings/Private/InternationalizationSettingsModelDetails.cpp
--- a/Engine/Source/Editor/InternationalizationSettings/Private/InternationalizationSettingsModelDetails.cpp
+++ b/Engine/Source/Editor/InternationalizationSettings/Private/InternationalizationSettingsModelDetails.cpp
@@ -83,87 +83,62 @@ namespace
 			LocalizedCulturesFlyweight = InLocalizedCulturesFlyweight;
 
 			ChildSlot
-				[
-					SAssignNew(EditorCultureComboButton, SComboButton)
-					.ButtonContent()
-					[
-						SNew(STextBlock)
-						.Text(this, &SEditorLanguageComboButton::GetContentText)
-					]
-				];
-			EditorCultureComboButton->SetOnGetMenuContent(FOnGetContent::CreateSP(this, &SEditorLanguageComboButton::GetMenuContent));
+			[
+				SNew(SCulturePickerCombo)
+				.Font(IPropertyTypeCustomizationUtils::GetRegularFont())
+				.SelectedCulture(this, &SEditorLanguageComboButton::GetSelectedCulture)
+				.OnSelectionChanged(this, &SEditorLanguageComboButton::OnSelectionChanged)
+				.IsCulturePickable(this, &SEditorLanguageComboButton::IsCulturePickable)
+				.DisplayNameFormat(SCulturePicker::ECultureDisplayFormat::ActiveAndNativeCultureDisplayName)
+				.ViewMode(SCulturePicker::ECulturesViewMode::Flat)
+			];
 		}
 
 	private:
 		TWeakObjectPtr<UInternationalizationSettingsModel> SettingsModel;
 		TSharedPtr<FLocalizedCulturesFlyweight> LocalizedCulturesFlyweight;
-		TSharedPtr<SComboButton> EditorCultureComboButton;
 
 	private:
-		FText GetContentText() const
+		FCulturePtr GetSelectedCulture() const
 		{
-			FCulturePtr EditorLanguage = FInternationalization::Get().GetCurrentLanguage();
-			return EditorLanguage.IsValid() ? FText::FromString(EditorLanguage->GetNativeName()) : LOCTEXT("None", "None");
+			return FInternationalization::Get().GetCurrentLanguage();
 		}
 
-		TSharedRef<SWidget> GetMenuContent()
+		void OnSelectionChanged(FCulturePtr SelectedCulture, ESelectInfo::Type SelectInfo)
 		{
-			FCulturePtr EditorLanguage = FInternationalization::Get().GetCurrentLanguage();
-
-			const auto& OnSelectionChangedLambda = [this](const FCulturePtr& SelectedCulture, ESelectInfo::Type SelectInfo)
+			if (SettingsModel.IsValid())
 			{
-				if (SettingsModel.IsValid())
+				SettingsModel->SetEditorLanguage(SelectedCulture.IsValid() ? SelectedCulture->GetName() : TEXT(""));
+				if (SelectedCulture.IsValid())
 				{
-					SettingsModel->SetEditorLanguage(SelectedCulture.IsValid() ? SelectedCulture->GetName() : TEXT(""));
-					if (SelectedCulture.IsValid())
+					FInternationalization& I18N = FInternationalization::Get();
+					I18N.SetCurrentLanguage(SelectedCulture->GetName());
+
+					// Find all Schemas and force a visualization cache clear
+					for (TObjectIterator<UClass> ClassIt; ClassIt; ++ClassIt)
 					{
-						FInternationalization& I18N = FInternationalization::Get();
-						I18N.SetCurrentLanguage(SelectedCulture->GetName());
+						UClass* CurrentClass = *ClassIt;
 
-						// Find all Schemas and force a visualization cache clear
-						for (TObjectIterator<UClass> ClassIt; ClassIt; ++ClassIt)
+						if (UEdGraphSchema* Schema = Cast<UEdGraphSchema>(CurrentClass->GetDefaultObject()))
 						{
-							UClass* CurrentClass = *ClassIt;
-
-							if (UEdGraphSchema* Schema = Cast<UEdGraphSchema>(CurrentClass->GetDefaultObject()))
-							{
-								Schema->ForceVisualizationCacheClear();
-							}
+							Schema->ForceVisualizationCacheClear();
 						}
 					}
 				}
-				if (EditorCultureComboButton.IsValid())
-				{
-					EditorCultureComboButton->SetIsOpen(false);
-				}
-			};
+			}
+		}
 
-			const auto& IsCulturePickableLambda = [this](FCulturePtr Culture) -> bool
+		bool IsCulturePickable(FCulturePtr Culture)
+		{
+			TArray<FString> CultureNames = Culture->GetPrioritizedParentCultureNames();
+			for (const FString& CultureName : CultureNames)
 			{
-				TArray<FString> CultureNames = Culture->GetPrioritizedParentCultureNames();
-				for (const FString& CultureName : CultureNames)
+				if (LocalizedCulturesFlyweight->LocalizedCulturesForEditor.Contains(Culture))
 				{
-					if (LocalizedCulturesFlyweight->LocalizedCulturesForEditor.Contains(Culture))
-					{
-						return true;
-					}
+					return true;
 				}
-				return false;
-			};
-
-			const auto& CulturePicker = SNew(SCulturePicker)
-				.InitialSelection(EditorLanguage)
-				.OnSelectionChanged_Lambda(OnSelectionChangedLambda)
-				.IsCulturePickable_Lambda(IsCulturePickableLambda)
-				.DisplayNameFormat(SCulturePicker::ECultureDisplayFormat::ActiveAndNativeCultureDisplayName)
-				.ViewMode(SCulturePicker::ECulturesViewMode::Flat);
-
-			return SNew(SBox)
-				.MaxDesiredHeight(500.0f)
-				.WidthOverride(350.0f)
-				[
-					CulturePicker
-				];
+			}
+			return false;
 		}
 	};
 
@@ -179,78 +154,47 @@ namespace
 			LocalizedCulturesFlyweight = InLocalizedCulturesFlyweight;
 
 			ChildSlot
-				[
-					SAssignNew(EditorCultureComboButton, SComboButton)
-					.ButtonContent()
-					[
-						SNew(STextBlock)
-						.Text(this, &SEditorLocaleComboButton::GetContentText)
-					]
-				];
-			EditorCultureComboButton->SetOnGetMenuContent(FOnGetContent::CreateSP(this, &SEditorLocaleComboButton::GetMenuContent));
+			[
+				SNew(SCulturePickerCombo)
+				.Font(IPropertyTypeCustomizationUtils::GetRegularFont())
+				.SelectedCulture(this, &SEditorLocaleComboButton::GetSelectedCulture)
+				.OnSelectionChanged(this, &SEditorLocaleComboButton::OnSelectionChanged)
+				.DisplayNameFormat(SCulturePicker::ECultureDisplayFormat::ActiveAndNativeCultureDisplayName)
+			];
 		}
 
 	private:
 		TWeakObjectPtr<UInternationalizationSettingsModel> SettingsModel;
 		TSharedPtr<FLocalizedCulturesFlyweight> LocalizedCulturesFlyweight;
-		TSharedPtr<SComboButton> EditorCultureComboButton;
 
 	private:
-		FText GetContentText() const
+		FCulturePtr GetSelectedCulture() const
 		{
-			FCulturePtr EditorLocale = FInternationalization::Get().GetCurrentLocale();
-			return EditorLocale.IsValid() ? FText::FromString(EditorLocale->GetNativeName()) : LOCTEXT("None", "None");
+			return FInternationalization::Get().GetCurrentLocale();
 		}
 
-		TSharedRef<SWidget> GetMenuContent()
+		void OnSelectionChanged(FCulturePtr SelectedCulture, ESelectInfo::Type SelectInfo)
 		{
-			FCulturePtr EditorLocale = FInternationalization::Get().GetCurrentLocale();
-
-			const auto& OnSelectionChangedLambda = [this](const FCulturePtr& SelectedCulture, ESelectInfo::Type SelectInfo)
+			if (SettingsModel.IsValid())
 			{
-				if (SettingsModel.IsValid())
+				SettingsModel->SetEditorLocale(SelectedCulture.IsValid() ? SelectedCulture->GetName() : TEXT(""));
+				if (SelectedCulture.IsValid())
 				{
-					SettingsModel->SetEditorLocale(SelectedCulture.IsValid() ? SelectedCulture->GetName() : TEXT(""));
-					if (SelectedCulture.IsValid())
+					FInternationalization& I18N = FInternationalization::Get();
+					I18N.SetCurrentLocale(SelectedCulture->GetName());
+
+					// Find all Schemas and force a visualization cache clear
+					for (TObjectIterator<UClass> ClassIt; ClassIt; ++ClassIt)
 					{
-						FInternationalization& I18N = FInternationalization::Get();
-						I18N.SetCurrentLocale(SelectedCulture->GetName());
+						UClass* CurrentClass = *ClassIt;
 
-						// Find all Schemas and force a visualization cache clear
-						for (TObjectIterator<UClass> ClassIt; ClassIt; ++ClassIt)
+						if (UEdGraphSchema* Schema = Cast<UEdGraphSchema>(CurrentClass->GetDefaultObject()))
 						{
-							UClass* CurrentClass = *ClassIt;
-
-							if (UEdGraphSchema* Schema = Cast<UEdGraphSchema>(CurrentClass->GetDefaultObject()))
-							{
-								Schema->ForceVisualizationCacheClear();
-							}
+							Schema->ForceVisualizationCacheClear();
 						}
 					}
 				}
-				if (EditorCultureComboButton.IsValid())
-				{
-					EditorCultureComboButton->SetIsOpen(false);
-				}
-			};
-
-			const auto& IsCulturePickableLambda = [](FCulturePtr Culture) -> bool
-			{
-				return true;
-			};
-
-			const auto& CulturePicker = SNew(SCulturePicker)
-				.InitialSelection(EditorLocale)
-				.OnSelectionChanged_Lambda(OnSelectionChangedLambda)
-				.IsCulturePickable_Lambda(IsCulturePickableLambda)
-				.DisplayNameFormat(SCulturePicker::ECultureDisplayFormat::ActiveAndNativeCultureDisplayName);
-
-			return SNew(SBox)
-				.MaxDesiredHeight(500.0f)
-				.WidthOverride(350.0f)
-				[
-					CulturePicker
-				];
+			}
 		}
 	};
 
@@ -266,98 +210,64 @@ namespace
 			LocalizedCulturesFlyweight = InLocalizedCulturesFlyweight;
 
 			ChildSlot
-				[
-					SAssignNew(PreviewGameCultureComboButton, SComboButton)
-					.ButtonContent()
-					[
-						SNew(STextBlock)
-						.Text(this, &SPreviewGameLanguageComboButton::GetContentText)
-					]
-				];
-			PreviewGameCultureComboButton->SetOnGetMenuContent(FOnGetContent::CreateSP(this, &SPreviewGameLanguageComboButton::GetMenuContent));
+			[
+				SNew(SCulturePickerCombo)
+				.Font(IPropertyTypeCustomizationUtils::GetRegularFont())
+				.SelectedCulture(this, &SPreviewGameLanguageComboButton::GetSelectedCulture)
+				.OnSelectionChanged(this, &SPreviewGameLanguageComboButton::OnSelectionChanged)
+				.IsCulturePickable(this, &SPreviewGameLanguageComboButton::IsCulturePickable)
+				.DisplayNameFormat(SCulturePicker::ECultureDisplayFormat::ActiveAndNativeCultureDisplayName)
+				.ViewMode(SCulturePicker::ECulturesViewMode::Flat)
+			];
 		}
 
 	private:
 		TWeakObjectPtr<UInternationalizationSettingsModel> SettingsModel;
 		TSharedPtr<FLocalizedCulturesFlyweight> LocalizedCulturesFlyweight;
-		TSharedPtr<SComboButton> PreviewGameCultureComboButton;
 
 	private:
-		FText GetContentText() const
+		FCulturePtr GetSelectedCulture() const
 		{
-			bool IsConfigured = false;
 			FString PreviewGameLanguage;
 			if (SettingsModel.IsValid())
 			{
-				IsConfigured = SettingsModel->GetPreviewGameLanguage(PreviewGameLanguage);
+				SettingsModel->GetPreviewGameLanguage(PreviewGameLanguage);
 			}
+
 			FCulturePtr Culture;
 			if (!PreviewGameLanguage.IsEmpty())
 			{
-				FInternationalization& I18N = FInternationalization::Get();
-				Culture = I18N.GetCulture(PreviewGameLanguage);
+				Culture = FInternationalization::Get().GetCulture(PreviewGameLanguage);
 			}
-			return Culture.IsValid() ? FText::FromString(Culture->GetDisplayName()) : LOCTEXT("None", "None");
+
+			return Culture;
 		}
 
-		TSharedRef<SWidget> GetMenuContent()
+		void OnSelectionChanged(FCulturePtr SelectedCulture, ESelectInfo::Type SelectInfo)
 		{
-			FCulturePtr PreviewGameCulture;
+			if (SettingsModel.IsValid())
 			{
-				bool IsConfigured = false;
-				FString PreviewGameLanguage;
-				if (SettingsModel.IsValid())
+				SettingsModel->SetPreviewGameLanguage(SelectedCulture.IsValid() ? SelectedCulture->GetName() : TEXT(""));
+
+				if (FTextLocalizationManager::Get().ShouldGameLocalizationPreviewAutoEnable() || FTextLocalizationManager::Get().IsGameLocalizationPreviewEnabled())
 				{
-					IsConfigured = SettingsModel->GetPreviewGameLanguage(PreviewGameLanguage);
-				}
-				if (!PreviewGameLanguage.IsEmpty())
-				{
-					FInternationalization& I18N = FInternationalization::Get();
-					PreviewGameCulture = I18N.GetCulture(PreviewGameLanguage);
+					// Enable the preview again for the newly set culture
+					FTextLocalizationManager::Get().EnableGameLocalizationPreview();
 				}
 			}
+		}
 
-			const auto& CulturePickerSelectLambda = [this](const FCulturePtr& SelectedCulture, ESelectInfo::Type SelectInfo)
-			{
-				if (SettingsModel.IsValid())
-				{
-					SettingsModel->SetPreviewGameLanguage(SelectedCulture.IsValid() ? SelectedCulture->GetName() : TEXT(""));
-
-					if (FTextLocalizationManager::Get().ShouldGameLocalizationPreviewAutoEnable() || FTextLocalizationManager::Get().IsGameLocalizationPreviewEnabled())
-					{
-						// Enable the preview again for the newly set culture
-						FTextLocalizationManager::Get().EnableGameLocalizationPreview();
-					}
-				}
-				if (PreviewGameCultureComboButton.IsValid())
-				{
-					PreviewGameCultureComboButton->SetIsOpen(false);
-				}
-			};
-			const auto& CulturePickerIsPickableLambda = [this](FCulturePtr Culture) -> bool
+		bool IsCulturePickable(FCulturePtr Culture)
+		{
+			TArray<FString> CultureNames = Culture->GetPrioritizedParentCultureNames();
+			for (const FString& CultureName : CultureNames)
 			{
-				TArray<FString> CultureNames = Culture->GetPrioritizedParentCultureNames();
-				for (const FString& CultureName : CultureNames)
+				if (LocalizedCulturesFlyweight->LocalizedCulturesForGame.Contains(Culture))
 				{
-					if (LocalizedCulturesFlyweight->LocalizedCulturesForGame.Contains(Culture))
-					{
-						return true;
-					}
+					return true;
 				}
-				return false;
-			};
-			return SNew(SBox)
-				.MaxDesiredHeight(500.0f)
-				.WidthOverride(350.0f)
-				[
-					SNew(SCulturePicker)
-					.InitialSelection(PreviewGameCulture)
-					.OnSelectionChanged_Lambda(CulturePickerSelectLambda)
-					.IsCulturePickable_Lambda(CulturePickerIsPickableLambda)
-					.DisplayNameFormat(SCulturePicker::ECultureDisplayFormat::ActiveAndNativeCultureDisplayName)
-					.ViewMode(SCulturePicker::ECulturesViewMode::Flat)
-					.CanSelectNone(true)
-				];
+			}
+			return false;
 		}
 	};
 }
diff --git a/Engine/Source/Editor/InternationalizationSettings/Private/SCulturePicker.cpp b/Engine/Source/Editor/InternationalizationSettings/Private/SCulturePicker.cpp
--- a/Engine/Source/Editor/InternationalizationSettings/Private/SCulturePicker.cpp
+++ b/Engine/Source/Editor/InternationalizationSettings/Private/SCulturePicker.cpp
@@ -15,6 +15,7 @@
 #include "Styling/SlateColor.h"
 #include "Templates/Function.h"
 #include "Templates/Tuple.h"
+#include "Widgets/Input/SComboButton.h"
 #include "Widgets/Input/SSearchBox.h"
 #include "Widgets/SBoxPanel.h"
 #include "Widgets/SToolTip.h"
@@ -25,6 +26,44 @@ class ITableRow;
 
 #define LOCTEXT_NAMESPACE "CulturePicker"
 
+namespace UE::Internationalization::Private
+{
+
+FString GetCultureDisplayName(const FCultureRef& Culture, SCulturePicker::ECultureDisplayFormat DisplayNameFormat, const bool bIsRootItem)
+{
+	const FString DisplayName = Culture->GetDisplayName();
+	if (DisplayNameFormat == SCulturePicker::ECultureDisplayFormat::ActiveCultureDisplayName)
+	{
+		return DisplayName;
+	}
+
+	const FString NativeName = Culture->GetNativeName();
+	if (DisplayNameFormat == SCulturePicker::ECultureDisplayFormat::NativeCultureDisplayName)
+	{
+		return NativeName;
+	}
+
+	if (DisplayNameFormat == SCulturePicker::ECultureDisplayFormat::ActiveAndNativeCultureDisplayName)
+	{
+		// Only show both names if they're different (to avoid repetition), and we're a root item (to avoid noise)
+		return (bIsRootItem && !NativeName.Equals(DisplayName, ESearchCase::CaseSensitive))
+			? FString::Printf(TEXT("%s [%s]"), *DisplayName, *NativeName)
+			: DisplayName;
+	}
+
+	if (DisplayNameFormat == SCulturePicker::ECultureDisplayFormat::NativeAndActiveCultureDisplayName)
+	{
+		// Only show both names if they're different (to avoid repetition), and we're a root item (to avoid noise)
+		return (bIsRootItem && !NativeName.Equals(DisplayName, ESearchCase::CaseSensitive))
+			? FString::Printf(TEXT("%s [%s]"), *NativeName, *DisplayName)
+			: NativeName;
+	}
+
+	return DisplayName;
+}
+
+}
+
 void SCulturePicker::Construct( const FArguments& InArgs )
 {
 	OnCultureSelectionChanged = InArgs._OnSelectionChanged;
@@ -326,35 +365,61 @@ void SCulturePicker::OnSelectionChanged(TSharedPtr<FCultureEntry> Entry, ESelect
 
 FString SCulturePicker::GetCultureDisplayName(const FCultureRef& Culture, const bool bIsRootItem) const
 {
-	const FString DisplayName = Culture->GetDisplayName();
-	if (DisplayNameFormat == ECultureDisplayFormat::ActiveCultureDisplayName)
-	{
-		return DisplayName;
-	}
+	return UE::Internationalization::Private::GetCultureDisplayName(Culture, DisplayNameFormat, bIsRootItem);
+}
 
-	const FString NativeName = Culture->GetNativeName();
-	if (DisplayNameFormat == ECultureDisplayFormat::NativeCultureDisplayName)
-	{
-		return NativeName;
-	}
+void SCulturePickerCombo::Construct(const FArguments& InArgs)
+{
+	SelectedCulture = InArgs._SelectedCulture;
+	OnSelectionChanged = InArgs._OnSelectionChanged;
+	DisplayNameFormat = InArgs._DisplayNameFormat;
+
+	ChildSlot
+	[
+		SAssignNew(ComboButton, SComboButton)
+		.ButtonContent()
+		[
+			SNew(STextBlock)
+			.Font(InArgs._Font)
+			.Text(this, &SCulturePickerCombo::GetSelectedCultureDisplayName)
+		]
+		.MenuContent()
+		[
+			SNew(SBox)
+			.MaxDesiredHeight(400.0f)
+			.MaxDesiredWidth(300.0f)
+			[
+				SNew(SCulturePicker)
+				.InitialSelection(SelectedCulture.Get(nullptr))
+				.OnSelectionChanged(this, &SCulturePickerCombo::OnSelectedCultureChanged)
+				.IsCulturePickable(InArgs._IsCulturePickable)
+				.DisplayNameFormat(DisplayNameFormat)
+				.ViewMode(InArgs._ViewMode)
+			]
+		]
+	];
+}
+
+FText SCulturePickerCombo::GetSelectedCultureDisplayName() const
+{
+	FCulturePtr SelectedCulturePtr = SelectedCulture.Get(nullptr);
+	return SelectedCulturePtr
+		? FText::AsCultureInvariant(UE::Internationalization::Private::GetCultureDisplayName(SelectedCulturePtr.ToSharedRef(), DisplayNameFormat, false))
+		: LOCTEXT("None", "None");
+}
 
-	if (DisplayNameFormat == ECultureDisplayFormat::ActiveAndNativeCultureDisplayName)
+void SCulturePickerCombo::OnSelectedCultureChanged(FCulturePtr NewSelectedCulture, ESelectInfo::Type SelectInfo)
+{
+	if (!SelectedCulture.IsBound())
 	{
-		// Only show both names if they're different (to avoid repetition), and we're a root item (to avoid noise)
-		return (bIsRootItem && !NativeName.Equals(DisplayName, ESearchCase::CaseSensitive))
-			? FString::Printf(TEXT("%s [%s]"), *DisplayName, *NativeName)
-			: DisplayName;
+		SelectedCulture = NewSelectedCulture;
 	}
+	OnSelectionChanged.ExecuteIfBound(NewSelectedCulture, SelectInfo);
 
-	if (DisplayNameFormat == ECultureDisplayFormat::NativeAndActiveCultureDisplayName)
+	if (ComboButton)
 	{
-		// Only show both names if they're different (to avoid repetition), and we're a root item (to avoid noise)
-		return (bIsRootItem && !NativeName.Equals(DisplayName, ESearchCase::CaseSensitive))
-			? FString::Printf(TEXT("%s [%s]"), *NativeName, *DisplayName)
-			: NativeName;
+		ComboButton->SetIsOpen(false);
 	}
-
-	return DisplayName;
 }
 
 #undef LOCTEXT_NAMESPACE
diff --git a/Engine/Source/Editor/InternationalizationSettings/Private/UserGeneratedContentLocalizationDescriptorDetails.cpp b/Engine/Source/Editor/InternationalizationSettings/Private/UserGeneratedContentLocalizationDescriptorDetails.cpp
--- a/Engine/Source/Editor/InternationalizationSettings/Private/UserGeneratedContentLocalizationDescriptorDetails.cpp
+++ b/Engine/Source/Editor/InternationalizationSettings/Private/UserGeneratedContentLocalizationDescriptorDetails.cpp
@@ -90,27 +90,13 @@ void FUserGeneratedContentLocalizationDescriptorDetails::CustomizeNativeCulture(
 		]
 		.ValueContent()
 		[
-			SAssignNew(NativeCulture_ComboButton,SComboButton)
-			.ButtonContent()
-			[
-				SNew(STextBlock)
-				.Font(IPropertyTypeCustomizationUtils::GetRegularFont())
-				.Text(this, &FUserGeneratedContentLocalizationDescriptorDetails::NativeCulture_GetDisplayName)
-			]
-			.MenuContent()
-			[
-				SNew(SBox)
-				.MaxDesiredHeight(400.0f)
-				.MaxDesiredWidth(300.0f)
-				[
-					SNew(SCulturePicker)
-					.InitialSelection(NativeCulture_GetCulture())
-					.OnSelectionChanged(this, &FUserGeneratedContentLocalizationDescriptorDetails::NativeCulture_OnSelectionChanged)
-					.IsCulturePickable(this, &FUserGeneratedContentLocalizationDescriptorDetails::NativeCulture_IsCulturePickable)
-					.DisplayNameFormat(SCulturePicker::ECultureDisplayFormat::ActiveCultureDisplayName)
-					.ViewMode(SCulturePicker::ECulturesViewMode::Flat)
-				]
-			]
+			SNew(SCulturePickerCombo)
+			.Font(IPropertyTypeCustomizationUtils::GetRegularFont())
+			.SelectedCulture(this, &FUserGeneratedContentLocalizationDescriptorDetails::NativeCulture_GetCulture)
+			.OnSelectionChanged(this, &FUserGeneratedContentLocalizationDescriptorDetails::NativeCulture_OnSelectionChanged)
+			.IsCulturePickable(this, &FUserGeneratedContentLocalizationDescriptorDetails::NativeCulture_IsCulturePickable)
+			.DisplayNameFormat(SCulturePicker::ECultureDisplayFormat::ActiveCultureDisplayName)
+			.ViewMode(SCulturePicker::ECulturesViewMode::Flat)
 		];
 }
 
@@ -124,15 +110,6 @@ FCulturePtr FUserGeneratedContentLocalizationDescriptorDetails::NativeCulture_Ge
 		: FInternationalization::Get().GetCulture(NativeCultureName);
 }
 
-FText FUserGeneratedContentLocalizationDescriptorDetails::NativeCulture_GetDisplayName() const
-{
-	FCulturePtr NativeCulturePtr = NativeCulture_GetCulture();
-
-	return NativeCulturePtr
-		? FText::AsCultureInvariant(NativeCulturePtr->GetDisplayName())
-		: LOCTEXT("None", "None");
-}
-
 bool FUserGeneratedContentLocalizationDescriptorDetails::NativeCulture_IsCulturePickable(FCulturePtr Culture) const
 {
 	return AvailableCultures.Contains(Culture);
@@ -141,11 +118,6 @@ bool FUserGeneratedContentLocalizationDescriptorDetails::NativeCulture_IsCulture
 void FUserGeneratedContentLocalizationDescriptorDetails::NativeCulture_OnSelectionChanged(FCulturePtr SelectedCulture, ESelectInfo::Type SelectInfo)
 {
 	NativeCultureHandle->SetValue(SelectedCulture->GetName());
-
-	if (NativeCulture_ComboButton)
-	{
-		NativeCulture_ComboButton->SetIsOpen(false);
-	}
 }
 
 void FUserGeneratedContentLocalizationDescriptorDetails::CustomizeCulturesToGenerate(TSharedRef<IPropertyHandle> PropertyHandle, IDetailChildrenBuilder& StructBuilder, IPropertyTypeCustomizationUtils& StructCustomizationUtils)
diff --git a/Engine/Source/Editor/InternationalizationSettings/Private/UserGeneratedContentLocalizationDescriptorDetails.h b/Engine/Source/Editor/InternationalizationSettings/Private/UserGeneratedContentLocalizationDescriptorDetails.h
--- a/Engine/Source/Editor/InternationalizationSettings/Private/UserGeneratedContentLocalizationDescriptorDetails.h
+++ b/Engine/Source/Editor/InternationalizationSettings/Private/UserGeneratedContentLocalizationDescriptorDetails.h
@@ -7,8 +7,6 @@
 
 namespace ESelectInfo { enum Type : int; }
 
-class SComboButton;
-
 class FUserGeneratedContentLocalizationDescriptorDetails : public IPropertyTypeCustomization
 {
 public:
@@ -26,7 +24,6 @@ public:
 private:
 	void CustomizeNativeCulture(TSharedRef<IPropertyHandle> PropertyHandle, IDetailChildrenBuilder& StructBuilder, IPropertyTypeCustomizationUtils& StructCustomizationUtils);
 	FCulturePtr NativeCulture_GetCulture() const;
-	FText NativeCulture_GetDisplayName() const;
 	bool NativeCulture_IsCulturePickable(FCulturePtr Culture) const;
 	void NativeCulture_OnSelectionChanged(FCulturePtr SelectedCulture, ESelectInfo::Type SelectInfo);
 
@@ -35,8 +32,6 @@ private:
 	void CulturesToGenerate_OnPostBatchSelect();
 	void CulturesToGenerate_OnCultureSelectionChanged(bool IsSelected, FCulturePtr Culture);
 	bool CulturesToGenerate_IsCultureSelected(FCulturePtr Culture) const;
-	void CulturesToGenerate_AddCulture(const FString& CultureName);
-	void CulturesToGenerate_RemoveCulture(const FString& CultureName);
 	
 	TArray<FCulturePtr> AvailableCultures;
 
@@ -44,7 +39,5 @@ private:
 
 	TSharedPtr<IPropertyHandle> CulturesToGenerateHandle;
 
-	TSharedPtr<SComboButton> NativeCulture_ComboButton;
-
 	bool CulturesToGenerate_IsInBatchSelectOperation = false;
 };
diff --git a/Engine/Source/Editor/InternationalizationSettings/Public/SCulturePicker.h b/Engine/Source/Editor/InternationalizationSettings/Public/SCulturePicker.h
--- a/Engine/Source/Editor/InternationalizationSettings/Public/SCulturePicker.h
+++ b/Engine/Source/Editor/InternationalizationSettings/Public/SCulturePicker.h
@@ -26,6 +26,7 @@
 
 class FText;
 class ITableRow;
+class SComboButton;
 
 struct FCultureEntry
 {
@@ -149,3 +150,33 @@ private:
 	/* Flags that the selection callback shouldn't be called when selecting - useful for initial selection, for instance. */
 	bool SuppressSelectionCallback;
 };
+
+class INTERNATIONALIZATIONSETTINGS_API SCulturePickerCombo : public SCompoundWidget
+{
+public:
+	SLATE_BEGIN_ARGS(SCulturePickerCombo)
+		: _DisplayNameFormat(SCulturePicker::ECultureDisplayFormat::ActiveCultureDisplayName)
+		, _ViewMode(SCulturePicker::ECulturesViewMode::Hierarchical)
+		, _CanSelectNone(false)
+	{}
+		SLATE_ATTRIBUTE(FSlateFontInfo, Font)
+		SLATE_EVENT(SCulturePicker::FOnSelectionChanged, OnSelectionChanged)
+		SLATE_EVENT(SCulturePicker::FIsCulturePickable, IsCulturePickable)
+		SLATE_ATTRIBUTE(FCulturePtr, SelectedCulture)
+		SLATE_ARGUMENT(SCulturePicker::ECultureDisplayFormat, DisplayNameFormat)
+		SLATE_ARGUMENT(SCulturePicker::ECulturesViewMode, ViewMode)
+		SLATE_ARGUMENT(bool, CanSelectNone)
+	SLATE_END_ARGS()
+
+	void Construct(const FArguments& InArgs);
+
+private:
+	FText GetSelectedCultureDisplayName() const;
+	void OnSelectedCultureChanged(FCulturePtr NewSelectedCulture, ESelectInfo::Type SelectInfo);
+
+	TAttribute<FCulturePtr> SelectedCulture;
+	SCulturePicker::FOnSelectionChanged OnSelectionChanged;
+	SCulturePicker::ECultureDisplayFormat DisplayNameFormat;
+
+	TSharedPtr<SComboButton> ComboButton;
+};
