Fixed being unabe to select "None" as the Game Preview Language

diff --git a/Engine/Source/Editor/InternationalizationSettings/Private/InternationalizationSettingsModelDetails.cpp b/Engine/Source/Editor/InternationalizationSettings/Private/InternationalizationSettingsModelDetails.cpp
--- a/Engine/Source/Editor/InternationalizationSettings/Private/InternationalizationSettingsModelDetails.cpp
+++ b/Engine/Source/Editor/InternationalizationSettings/Private/InternationalizationSettingsModelDetails.cpp
@@ -218,6 +218,7 @@ namespace
 				.IsCulturePickable(this, &SPreviewGameLanguageComboButton::IsCulturePickable)
 				.DisplayNameFormat(SCulturePicker::ECultureDisplayFormat::ActiveAndNativeCultureDisplayName)
 				.ViewMode(SCulturePicker::ECulturesViewMode::Flat)
+				.CanSelectNone(true)
 			];
 		}
 
diff --git a/Engine/Source/Editor/InternationalizationSettings/Private/SCulturePicker.cpp b/Engine/Source/Editor/InternationalizationSettings/Private/SCulturePicker.cpp
--- a/Engine/Source/Editor/InternationalizationSettings/Private/SCulturePicker.cpp
+++ b/Engine/Source/Editor/InternationalizationSettings/Private/SCulturePicker.cpp
@@ -395,6 +395,7 @@ void SCulturePickerCombo::Construct(const FArguments& InArgs)
 				.IsCulturePickable(InArgs._IsCulturePickable)
 				.DisplayNameFormat(DisplayNameFormat)
 				.ViewMode(InArgs._ViewMode)
+				.CanSelectNone(InArgs._CanSelectNone)
 			]
 		]
 	];
